Curve Editor - Add snapping support.

diff --git a/Engine/Source/Editor/CurveAssetEditor/CurveAssetEditor.Build.cs b/Engine/Source/Editor/CurveAssetEditor/CurveAssetEditor.Build.cs
--- a/Engine/Source/Editor/CurveAssetEditor/CurveAssetEditor.Build.cs
+++ b/Engine/Source/Editor/CurveAssetEditor/CurveAssetEditor.Build.cs
@@ -17,7 +17,8 @@ public class CurveAssetEditor : ModuleRules
 				"Slate",
 				"SlateCore",
                 "EditorStyle",
-				"UnrealEd"
+				"UnrealEd",
+				"InputCore"
 			}
 		);
 
diff --git a/Engine/Source/Editor/CurveAssetEditor/Private/CurveAssetEditor.cpp b/Engine/Source/Editor/CurveAssetEditor/Private/CurveAssetEditor.cpp
--- a/Engine/Source/Editor/CurveAssetEditor/Private/CurveAssetEditor.cpp
+++ b/Engine/Source/Editor/CurveAssetEditor/Private/CurveAssetEditor.cpp
@@ -99,6 +99,9 @@ TSharedRef<SDockTab> FCurveAssetEditor::SpawnTab_CurveAsset( const FSpawnTabArgs
 	ViewMinInput=0.f;
 	ViewMaxInput=5.f;
 
+	InputSnap = 0.1f;
+	OutputSnap = 0.05f;
+
 	TSharedRef<SDockTab> NewDockTab = SNew(SDockTab)
 		.Icon( FEditorStyle::GetBrush("CurveAssetEditor.Tabs.Properties") )
 		.Label( FText::Format(LOCTEXT("CurveAssetEditorTitle", "{0} Curve Asset"), FText::FromString(GetTabPrefix())))
@@ -111,6 +114,8 @@ TSharedRef<SDockTab> FCurveAssetEditor::SpawnTab_CurveAsset( const FSpawnTabArgs
 				SAssignNew(TrackWidget, SCurveEditor)
 				.ViewMinInput(this, &FCurveAssetEditor::GetViewMinInput)
 				.ViewMaxInput(this, &FCurveAssetEditor::GetViewMaxInput)
+				.InputSnap(this, &FCurveAssetEditor::GetInputSnap)
+				.OutputSnap(this, &FCurveAssetEditor::GetOutputSnap)
 				.TimelineLength(this, &FCurveAssetEditor::GetTimelineLength)
 				.OnSetInputViewRange(this, &FCurveAssetEditor::SetInputViewRange)
 				.HideUI(false)
@@ -148,6 +153,102 @@ TSharedRef<SDockTab> FCurveAssetEditor::SpawnTab_CurveAsset( const FSpawnTabArgs
 	return NewDockTab;
 }
 
+float FCurveAssetEditor::GetInputSnap() const
+{
+	return InputSnap;
+}
+
+void FCurveAssetEditor::SetInputSnap(float value)
+{
+	InputSnap = value;
+}
+
+FText FCurveAssetEditor::GetInputSnapText() const
+{
+	return FText::AsNumber(InputSnap);
+}
+
+void FCurveAssetEditor::InputSnapTextComitted(const FText& InNewText, ETextCommit::Type InTextCommit)
+{
+	if (InNewText.IsNumeric())
+	{
+		TTypeFromString<float>::FromString(InputSnap, *InNewText.ToString());
+	}
+}
+
+float FCurveAssetEditor::GetOutputSnap() const
+{
+	return OutputSnap;
+}
+
+void FCurveAssetEditor::SetOutputSnap(float value)
+{
+	OutputSnap = value;
+}
+
+FText FCurveAssetEditor::GetOutputSnapText() const
+{
+	return FText::AsNumber(OutputSnap);
+}
+
+void FCurveAssetEditor::OutputSnapTextComitted(const FText& InNewText, ETextCommit::Type InTextCommit)
+{
+	if (InNewText.IsNumeric())
+	{
+		TTypeFromString<float>::FromString(OutputSnap, *InNewText.ToString());
+	}
+}
+
+TSharedRef<SWidget> FCurveAssetEditor::BuildInputSnapMenu()
+{
+	FMenuBuilder MenuBuilder(true, NULL);
+
+	FUIAction OneThousandthAction(FExecuteAction::CreateSP(this, &FCurveAssetEditor::SetInputSnap, 0.001f));
+	MenuBuilder.AddMenuEntry(LOCTEXT("InputSnap_OneThousandth", "0.001"), LOCTEXT("InputSnap_OneThousandth", "Set snap to 1/1000th"), FSlateIcon(), OneThousandthAction);
+
+	FUIAction OneHundredthAction(FExecuteAction::CreateSP(this, &FCurveAssetEditor::SetInputSnap, 0.01f));
+	MenuBuilder.AddMenuEntry(LOCTEXT("InputSnap_OneHundredth", "0.01"), LOCTEXT("InputSnap_OneHundredth", "Set snap to 1/100th"), FSlateIcon(), OneHundredthAction);
+
+	FUIAction OneTenthAction(FExecuteAction::CreateSP(this, &FCurveAssetEditor::SetInputSnap, 0.1f));
+	MenuBuilder.AddMenuEntry(LOCTEXT("InputSnap_OneTenth", "0.1"), LOCTEXT("InputSnap_OneTenth", "Set snap to 1/10th"), FSlateIcon(), OneTenthAction);
+
+	FUIAction OneAction(FExecuteAction::CreateSP(this, &FCurveAssetEditor::SetInputSnap, 1.0f));
+	MenuBuilder.AddMenuEntry(LOCTEXT("InputSnap_One", "1"), LOCTEXT("InputSnap_One", "Set snap to 1h"), FSlateIcon(), OneAction);
+
+	FUIAction TenAction(FExecuteAction::CreateSP(this, &FCurveAssetEditor::SetInputSnap, 10.0f));
+	MenuBuilder.AddMenuEntry(LOCTEXT("InputSnap_Ten", "10"), LOCTEXT("InputSnap_Ten", "Set snap to 10"), FSlateIcon(), TenAction);
+
+	FUIAction HundredAction(FExecuteAction::CreateSP(this, &FCurveAssetEditor::SetInputSnap, 100.0f));
+	MenuBuilder.AddMenuEntry(LOCTEXT("InputSnap_OneHundred", "100"), LOCTEXT("InputSnap_OneHundred", "Set snap to 100"), FSlateIcon(), HundredAction);
+
+	return MenuBuilder.MakeWidget();
+}
+
+TSharedRef<SWidget> FCurveAssetEditor::BuildOutputSnapMenu()
+{
+	FMenuBuilder MenuBuilder(true, NULL);
+
+	FUIAction OneThousandthAction(FExecuteAction::CreateSP(this, &FCurveAssetEditor::SetOutputSnap, 0.001f));
+	MenuBuilder.AddMenuEntry(LOCTEXT("OutputSnap_OneThousandth", "0.001"), LOCTEXT("OutputSnap_OneThousandth", "Set snap to 1/1000th"), FSlateIcon(), OneThousandthAction);
+
+	FUIAction OneHundredthAction(FExecuteAction::CreateSP(this, &FCurveAssetEditor::SetOutputSnap, 0.01f));
+	MenuBuilder.AddMenuEntry(LOCTEXT("OutputSnap_OneHundredth", "0.01"), LOCTEXT("OutputSnap_OneHundredth", "Set snap to 1/100th"), FSlateIcon(), OneHundredthAction);
+
+	FUIAction OneTenthAction(FExecuteAction::CreateSP(this, &FCurveAssetEditor::SetOutputSnap, 0.1f));
+	MenuBuilder.AddMenuEntry(LOCTEXT("OutputSnap_OneTenth", "0.1"), LOCTEXT("OutputSnap_OneTenth", "Set snap to 1/10th"), FSlateIcon(), OneTenthAction);
+
+	FUIAction OneAction(FExecuteAction::CreateSP(this, &FCurveAssetEditor::SetOutputSnap, 1.0f));
+	MenuBuilder.AddMenuEntry(LOCTEXT("OutputSnap_One", "1"), LOCTEXT("OutputSnap_One", "Set snap to 1h"), FSlateIcon(), OneAction);
+
+	FUIAction TenAction(FExecuteAction::CreateSP(this, &FCurveAssetEditor::SetOutputSnap, 10.0f));
+	MenuBuilder.AddMenuEntry(LOCTEXT("OutputSnap_Ten", "10"), LOCTEXT("OutputSnap_Ten", "Set snap to 10"), FSlateIcon(), TenAction);
+
+	FUIAction HundredAction(FExecuteAction::CreateSP(this, &FCurveAssetEditor::SetOutputSnap, 100.0f));
+	MenuBuilder.AddMenuEntry(LOCTEXT("OutputSnap_OneHundred", "100"), LOCTEXT("OutputSnap_OneHundred", "Set snap to 100"), FSlateIcon(), HundredAction);
+
+	return MenuBuilder.MakeWidget();
+}
+
 float FCurveAssetEditor::GetTimelineLength() const
 {
 	return 0.f;
@@ -163,7 +264,7 @@ TSharedPtr<FExtender> FCurveAssetEditor::GetToolbarExtender()
 {
 	struct Local
 	{
-		static void FillToolbar(FToolBarBuilder& ToolbarBuilder)
+		static void FillToolbar(FToolBarBuilder& ToolbarBuilder, TSharedRef<SWidget> InputSnapWidget, TSharedRef<SWidget> OutputSnapWidget)
 		{
 			ToolbarBuilder.BeginSection("Curve");
 			{
@@ -171,16 +272,71 @@ TSharedPtr<FExtender> FCurveAssetEditor::GetToolbarExtender()
 				ToolbarBuilder.AddToolBarButton(FRichCurveEditorCommands::Get().ZoomToFitVertical);
 			}
 			ToolbarBuilder.EndSection();
+			ToolbarBuilder.BeginSection("Snap");
+			{
+				ToolbarBuilder.AddToolBarButton(FRichCurveEditorCommands::Get().ToggleSnapping);
+				ToolbarBuilder.AddWidget(InputSnapWidget);
+				ToolbarBuilder.AddWidget(OutputSnapWidget);
+			}
+			ToolbarBuilder.EndSection();
 		}
 	};
 
 	TSharedPtr<FExtender> ToolbarExtender = MakeShareable(new FExtender);
 
+	TSharedRef<SWidget> InputSnapWidget =
+		SNew(SVerticalBox)
+		+ SVerticalBox::Slot()
+		.Padding(4)
+		.AutoHeight()
+		[
+			SNew(STextBlock)
+			.Text(LOCTEXT("InputSnap", "Time Snap"))
+		]
+		+ SVerticalBox::Slot()
+		.Padding(4)
+		.AutoHeight()
+		[
+			SNew(SComboButton)
+			.ContentPadding(1)
+			.OnGetMenuContent(this, &FCurveAssetEditor::BuildInputSnapMenu)
+			.ButtonContent()
+			[
+				SNew(SEditableTextBox)
+				.Text(this, &FCurveAssetEditor::GetInputSnapText)
+				.OnTextCommitted(this, &FCurveAssetEditor::InputSnapTextComitted)
+			]
+		];
+
+	TSharedRef<SWidget> OutputSnapWidget =
+		SNew(SVerticalBox)
+		+ SVerticalBox::Slot()
+		.Padding(4)
+		.AutoHeight()
+		[
+			SNew(STextBlock)
+			.Text(LOCTEXT("OutputSnap", "Value Snap"))
+		]
+		+ SVerticalBox::Slot()
+		.Padding(4)
+		.AutoHeight()
+		[
+			SNew(SComboButton)
+			.ContentPadding(1)
+			.OnGetMenuContent(this, &FCurveAssetEditor::BuildOutputSnapMenu)
+			.ButtonContent()
+			[
+				SNew(SEditableTextBox)
+				.Text(this, &FCurveAssetEditor::GetOutputSnapText)
+				.OnTextCommitted(this, &FCurveAssetEditor::OutputSnapTextComitted)
+			]
+		];
+
 	ToolbarExtender->AddToolBarExtension(
 		"Asset",
 		EExtensionHook::After,
 		TrackWidget->GetCommands(),
-		FToolBarExtensionDelegate::CreateStatic(&Local::FillToolbar)
+		FToolBarExtensionDelegate::CreateStatic(&Local::FillToolbar, InputSnapWidget, OutputSnapWidget)
 		);
 
 	return ToolbarExtender;
diff --git a/Engine/Source/Editor/CurveAssetEditor/Private/CurveAssetEditor.h b/Engine/Source/Editor/CurveAssetEditor/Private/CurveAssetEditor.h
--- a/Engine/Source/Editor/CurveAssetEditor/Private/CurveAssetEditor.h
+++ b/Engine/Source/Editor/CurveAssetEditor/Private/CurveAssetEditor.h
@@ -47,6 +47,28 @@ private:
 	/**	Spawns the tab with the curve asset inside */
 	TSharedRef<SDockTab> SpawnTab_CurveAsset( const FSpawnTabArgs& Args );
 
+	/** Get the snapping value for the input domain. */
+	float GetInputSnap() const;
+	/** Gets the snapping value for the input domain as text. */
+	FText GetInputSnapText() const;
+	/** Get set the snapping value for the input domain. */
+	void SetInputSnap(float value);
+	/** Callback for setting the snapping value for the input domain from a text value. */
+	void InputSnapTextComitted(const FText& InNewText, ETextCommit::Type InTextCommit);
+	/** Builds a menu widget with options for snapping in the input domain. */
+	TSharedRef<SWidget> BuildInputSnapMenu();
+
+	/** Get the snapping value for the output domain. */
+	float GetOutputSnap() const;
+	/** Gets the snapping value for the output domain as text. */
+	FText GetOutputSnapText() const;
+	/** Get set the snapping value for the output domain. */
+	void SetOutputSnap(float value);
+	/** Callback for setting the snapping value for the output domain from a text value. */
+	void OutputSnapTextComitted(const FText& InNewText, ETextCommit::Type InTextCommit);
+	/** Builds a menu widget with options for snapping in the output domain. */
+	TSharedRef<SWidget> BuildOutputSnapMenu();
+
 	TSharedPtr<class SCurveEditor> TrackWidget;
 
 	/**	The tab id for the curve asset tab */
@@ -54,4 +76,9 @@ private:
 
 	float ViewMinInput;
 	float ViewMaxInput;
+
+	/** The snapping value for the input domain. */
+	float InputSnap;
+	/** The snapping value for the output domain. */
+	float OutputSnap;
 };
diff --git a/Engine/Source/Editor/EditorStyle/Private/SlateEditorStyle.cpp b/Engine/Source/Editor/EditorStyle/Private/SlateEditorStyle.cpp
--- a/Engine/Source/Editor/EditorStyle/Private/SlateEditorStyle.cpp
+++ b/Engine/Source/Editor/EditorStyle/Private/SlateEditorStyle.cpp
@@ -5208,6 +5208,7 @@ void FSlateEditorStyle::FStyle::SetupToolkitStyles()
 	{
 		Set("RichCurveEditor.ZoomToFitHorizontal", new IMAGE_BRUSH("Icons/icon_CurveEditor_Horizontal_40x", Icon40x40));
 		Set("RichCurveEditor.ZoomToFitVertical", new IMAGE_BRUSH("Icons/icon_CurveEditor_Vertical_40x", Icon40x40));
+		Set("RichCurveEditor.ToggleSnapping", new IMAGE_BRUSH("Icons/icon_CurveEditor_ToggleSnap_40x", Icon40x40));
 	}
 
 	// PhAT
diff --git a/Engine/Source/Editor/UnrealEd/Private/RichCurveEditorCommands.cpp b/Engine/Source/Editor/UnrealEd/Private/RichCurveEditorCommands.cpp
--- a/Engine/Source/Editor/UnrealEd/Private/RichCurveEditorCommands.cpp
+++ b/Engine/Source/Editor/UnrealEd/Private/RichCurveEditorCommands.cpp
@@ -5,4 +5,5 @@ void FRichCurveEditorCommands::RegisterCommands()
 {
 	UI_COMMAND(ZoomToFitHorizontal, "Fit Horizontal", "Zoom to Fit - Horizontal", EUserInterfaceActionType::Button, FInputGesture());
 	UI_COMMAND(ZoomToFitVertical, "Fit Vertical", "Zoom to Fit - Vertical", EUserInterfaceActionType::Button, FInputGesture());
+	UI_COMMAND(ToggleSnapping, "Snapping", "Toggle Snapping", EUserInterfaceActionType::ToggleButton, FInputGesture());
 }
diff --git a/Engine/Source/Editor/UnrealEd/Private/SCurveEditor.cpp b/Engine/Source/Editor/UnrealEd/Private/SCurveEditor.cpp
--- a/Engine/Source/Editor/UnrealEd/Private/SCurveEditor.cpp
+++ b/Engine/Source/Editor/UnrealEd/Private/SCurveEditor.cpp
@@ -37,6 +37,10 @@ void SCurveEditor::Construct(const FArguments& InArgs)
 	ViewMinOutput = InArgs._ViewMinOutput;
 	ViewMaxOutput = InArgs._ViewMaxOutput;
 
+	InputSnap = InArgs._InputSnap;
+	OutputSnap = InArgs._OutputSnap;
+	bSnappingEnabled = InArgs._SnappingEnabled;
+
 	bZoomToFitVertical = InArgs._ZoomToFitVertical;
 	bZoomToFitHorizontal = InArgs._ZoomToFitHorizontal;
 	DesiredSize = InArgs._DesiredSize;
@@ -81,6 +85,11 @@ void SCurveEditor::Construct(const FArguments& InArgs)
 	Commands->MapAction(FRichCurveEditorCommands::Get().ZoomToFitVertical,
 		FExecuteAction::CreateSP(this, &SCurveEditor::ZoomToFitVertical));
 
+	Commands->MapAction(FRichCurveEditorCommands::Get().ToggleSnapping,
+		FExecuteAction::CreateSP(this, &SCurveEditor::ToggleSnapping),
+		FCanExecuteAction(),
+		FIsActionChecked::CreateSP(this, &SCurveEditor::IsSnappingEnabled));
+
 	SAssignNew(WarningMessageText, SErrorText);
 
 	TSharedRef<SBox> CurveSelector = SNew(SBox)
@@ -1040,7 +1049,7 @@ FReply SCurveEditor::OnMouseMove( const FGeometry& InMyGeometry, const FPointerE
 		{
 			TryStartDrag(InMyGeometry, InMouseEvent);
 		}
-		else if (DragState != EDragState::None)
+		if (DragState != EDragState::None)
 		{
 			ProcessDrag(InMyGeometry, InMouseEvent);
 		}
@@ -1119,6 +1128,16 @@ void SCurveEditor::TryStartDrag(const FGeometry& InMyGeometry, const FPointerEve
 
 				BeginDragTransaction();
 				DragState = EDragState::DragKey;
+				DraggedKeyHandle = HitKey.KeyHandle;
+				PreDragKeyLocations.Empty();
+				for (auto selectedKey : SelectedKeys)
+				{
+					PreDragKeyLocations.Add(selectedKey.KeyHandle, FVector2D
+					(
+						selectedKey.Curve->GetKeyTime(selectedKey.KeyHandle),
+						selectedKey.Curve->GetKeyValue(selectedKey.KeyHandle)
+					));
+				}
 			}
 			else
 			{
@@ -1159,7 +1178,8 @@ void SCurveEditor::ProcessDrag(const FGeometry& InMyGeometry, const FPointerEven
 
 	if (DragState == EDragState::DragKey)
 	{
-		MoveSelectedKeysByDelta(InputDelta);
+		FVector2D MousePosition = InMyGeometry.AbsoluteToLocal(InMouseEvent.GetScreenSpacePosition());
+		MoveSelectedKeys(FVector2D(ScaleInfo.LocalXToInput(MousePosition.X), ScaleInfo.LocalYToOutput(MousePosition.Y)));
 	}
 	else if (DragState == EDragState::DragTangent)
 	{
@@ -1481,7 +1501,7 @@ SCurveEditor::FSelectedCurveKey SCurveEditor::HitTestKeys(const FGeometry& InMyG
 	return SelectedKey;
 }
 
-void SCurveEditor::MoveSelectedKeysByDelta(FVector2D InputDelta)
+void SCurveEditor::MoveSelectedKeys(FVector2D NewLocation)
 {
 	const FScopedTransaction Transaction( LOCTEXT("CurveEditor_MoveKeys", "Move Keys") );
 	CurveOwner->ModifyOwner();
@@ -1489,29 +1509,35 @@ void SCurveEditor::MoveSelectedKeysByDelta(FVector2D InputDelta)
 	// track all unique curves encountered so their tangents can be updated later
 	TSet<FRichCurve*> UniqueCurves;
 
-	for(int32 i=0; i<SelectedKeys.Num(); i++)
+	FVector2D SnappedNewLocation = SnapLocation(NewLocation);
+
+	// The total move distance for all keys is the difference between the current snapped location
+	// and the start location of the key which was actually dragged.
+	FVector2D TotalMoveDistance = SnappedNewLocation - PreDragKeyLocations[DraggedKeyHandle];
+
+	for (int32 i = 0; i < SelectedKeys.Num(); i++)
 	{
 		FSelectedCurveKey OldKey = SelectedKeys[i];
 
-		if(!IsValidCurve(OldKey.Curve))
+		if (!IsValidCurve(OldKey.Curve))
 		{
 			continue;
 		}
 
-		FKeyHandle OldKeyHandle	= OldKey.KeyHandle;
-		FRichCurve* Curve		= OldKey.Curve;
+		FKeyHandle OldKeyHandle = OldKey.KeyHandle;
+		FRichCurve* Curve = OldKey.Curve;
 
-		// Update output - easy
-		float OldOutput = Curve->GetKeyValue(OldKeyHandle);
-		// set the key value, but don't update the tangents yet
-		Curve->SetKeyValue(OldKeyHandle, OldOutput + InputDelta.Y, false);
+		FVector2D PreDragLocation = PreDragKeyLocations[OldKeyHandle];
+		FVector2D NewLocation = PreDragLocation + TotalMoveDistance;
 
-		// Update input - hard
-		float OldInput = Curve->GetKeyTime(OldKeyHandle);
-		FKeyHandle KeyHandle = Curve->SetKeyTime(OldKeyHandle, OldInput + InputDelta.X);
+		// Update the key's value without updating the tangents.
+		Curve->SetKeyValue(OldKeyHandle, NewLocation.Y, false);
 
-		// Update this key
+		// Changing the time of a key returns a new handle, so make sure to update existing references.
+		FKeyHandle KeyHandle = Curve->SetKeyTime(OldKeyHandle, NewLocation.X);
 		SelectedKeys[i] = FSelectedCurveKey(Curve, KeyHandle);
+		PreDragKeyLocations.Remove(OldKeyHandle);
+		PreDragKeyLocations.Add(KeyHandle, PreDragLocation);
 
 		UniqueCurves.Add(Curve);
 	}
@@ -1687,6 +1713,16 @@ FReply SCurveEditor::ZoomToFitVerticalClicked()
 	return FReply::Handled();
 }
 
+void SCurveEditor::ToggleSnapping()
+{
+	bSnappingEnabled = !bSnappingEnabled;
+}
+
+bool SCurveEditor::IsSnappingEnabled()
+{
+	return bSnappingEnabled;
+}
+
 void SCurveEditor::CreateContextMenu(const FGeometry& InMyGeometry, const FVector2D& ScreenPosition)
 {
 	const bool CloseAfterSelection = true;
@@ -2219,6 +2255,16 @@ bool SCurveEditor::IsLinearColorCurve() const
 	return CurveOwner && CurveOwner->GetOwner() && CurveOwner->GetOwner()->IsA<UCurveLinearColor>();
 }
 
+FVector2D SCurveEditor::SnapLocation(FVector2D InLocation)
+{
+	if (bSnappingEnabled)
+	{
+		InLocation.X = InputSnap != 0 ? FMath::RoundToInt(InLocation.X / InputSnap.Get()) * InputSnap.Get() : InLocation.X;
+		InLocation.Y = OutputSnap != 0 ? FMath::RoundToInt(InLocation.Y / OutputSnap.Get()) * OutputSnap.Get() : InLocation.Y;
+	}
+	return InLocation;
+}
+
 FText SCurveEditor::GetInterpolationModeText() const
 {
 	FString Results;
diff --git a/Engine/Source/Editor/UnrealEd/Public/RichCurveEditorCommands.h b/Engine/Source/Editor/UnrealEd/Public/RichCurveEditorCommands.h
--- a/Engine/Source/Editor/UnrealEd/Public/RichCurveEditorCommands.h
+++ b/Engine/Source/Editor/UnrealEd/Public/RichCurveEditorCommands.h
@@ -19,12 +19,8 @@ public:
 
 	TSharedPtr<FUICommandInfo> ZoomToFitHorizontal;
 	TSharedPtr<FUICommandInfo> ZoomToFitVertical;
+	TSharedPtr<FUICommandInfo> ToggleSnapping;
 
 public:
 	virtual void RegisterCommands() override;
-
-	~FRichCurveEditorCommands()
-	{
-		int i = 0;
-	}
 };
\ No newline at end of file
diff --git a/Engine/Source/Editor/UnrealEd/Public/SCurveEditor.h b/Engine/Source/Editor/UnrealEd/Public/SCurveEditor.h
--- a/Engine/Source/Editor/UnrealEd/Public/SCurveEditor.h
+++ b/Engine/Source/Editor/UnrealEd/Public/SCurveEditor.h
@@ -106,6 +106,9 @@ public:
 		, _ViewMaxInput(10.0f)
 		, _ViewMinOutput(0.0f)
 		, _ViewMaxOutput(1.0f)
+		, _InputSnap(0.1f)
+		, _OutputSnap(0.05f)
+		, _SnappingEnabled(false)
 		, _TimelineLength(5.0f)
 		, _DesiredSize(FVector2D::ZeroVector)
 		, _DrawCurve(true)
@@ -123,6 +126,9 @@ public:
 		SLATE_ATTRIBUTE( TOptional<float>, DataMaxInput )
 		SLATE_ARGUMENT( float, ViewMinOutput )
 		SLATE_ARGUMENT( float, ViewMaxOutput )
+		SLATE_ATTRIBUTE( float, InputSnap )
+		SLATE_ATTRIBUTE( float, OutputSnap )
+		SLATE_ARGUMENT( bool, SnappingEnabled )
 		SLATE_ATTRIBUTE( float, TimelineLength )
 		SLATE_ATTRIBUTE( FVector2D, DesiredSize )
 		SLATE_ARGUMENT( bool, DrawCurve )
@@ -273,8 +279,8 @@ private:
 	/** Get the time of the desired key */
 	TOptional<float> GetKeyTime(FSelectedCurveKey Key) const;
 
-	/** Move the selected keys by the supplied change in input (X) */
-	void MoveSelectedKeysByDelta(FVector2D InputDelta);
+	/** Move the selected keys */
+	void MoveSelectedKeys(FVector2D NewLocation);
 
 	/** Function to check whether the current track is editable */
 	bool IsEditingEnabled() const;
@@ -288,6 +294,9 @@ private:
 	FReply ZoomToFitHorizontalClicked();
 	FReply ZoomToFitVerticalClicked();
 
+	void ToggleSnapping();
+	bool IsSnappingEnabled();
+
 	TOptional<float> OnGetTime() const;
 	void OnTimeComitted(float NewValue, ETextCommit::Type CommitType);
 	void OnTimeChanged(float NewValue);
@@ -428,6 +437,8 @@ private:
 	bool IsGradientEditorVisible() const { return bIsGradientEditorVisible; }
 	bool IsLinearColorCurve() const;
 
+	FVector2D SnapLocation(FVector2D InLocation);
+
 protected:
 
 	/** Set Default output values when range is too small **/
@@ -520,6 +531,15 @@ protected:
 	/** Min output view range */
 	float				ViewMaxOutput;
 
+	/** The snapping value for the input domain. */
+	TAttribute<float> InputSnap;
+
+	/** The snapping value for the output domain. */
+	TAttribute<float> OutputSnap;
+
+	/** Whether or not snapping is enabled. */
+	bool bSnappingEnabled;
+
 	/** True if you want the curve editor to fit to zoom **/
 	bool				bZoomToFitVertical;
 
@@ -543,6 +563,12 @@ protected:
 
 	/** The number of pixels which the mouse must move before a drag operation starts. */
 	float DragThreshold;
+
+	/** A handle to the key which was clicked to start a key drag operation. */
+	FKeyHandle DraggedKeyHandle;
+
+	/** A map of selected key handle to their starting locations at the beginning of a drag operation. */
+	TMap<FKeyHandle, FVector2D> PreDragKeyLocations;
 };
 
 #endif // __SCurveEditor_h__
