Curve Asset Editor - Fix the toolbar so the icons don't all dissapear in small icon mode.

diff --git a/Engine/Source/Editor/CurveAssetEditor/Private/CurveAssetEditor.cpp b/Engine/Source/Editor/CurveAssetEditor/Private/CurveAssetEditor.cpp
--- a/Engine/Source/Editor/CurveAssetEditor/Private/CurveAssetEditor.cpp
+++ b/Engine/Source/Editor/CurveAssetEditor/Private/CurveAssetEditor.cpp
@@ -304,49 +304,85 @@ TSharedPtr<FExtender> FCurveAssetEditor::GetToolbarExtender()
 
 	TSharedRef<SWidget> InputSnapWidget =
 		SNew(SVerticalBox)
+		// Vertical Label
 		+ SVerticalBox::Slot()
 		.Padding(4)
 		.AutoHeight()
 		[
 			SNew(STextBlock)
-			.Text(LOCTEXT("InputSnap", "Time Snap"))
+			.Text(LOCTEXT("InputSnapLabelVertical", "Time Snap"))
+			.Visibility(this, &FCurveAssetEditor::GetSnapLabelVisibilty, EOrientation::Orient_Vertical)
 		]
 		+ SVerticalBox::Slot()
 		.Padding(4)
 		.AutoHeight()
 		[
-			SNew(SComboButton)
-			.ContentPadding(1)
-			.OnGetMenuContent(this, &FCurveAssetEditor::BuildInputSnapMenu)
-			.ButtonContent()
+			SNew(SHorizontalBox)
+			// Horizontal Icon Label
+			+ SHorizontalBox::Slot()
+			.VAlign(EVerticalAlignment::VAlign_Center)
+			.AutoWidth()
+			.Padding(FMargin(0, 0, 3, 0))
 			[
-				SNew(SEditableTextBox)
-				.Text(this, &FCurveAssetEditor::GetInputSnapText)
-				.OnTextCommitted(this, &FCurveAssetEditor::InputSnapTextComitted)
+				SNew( STextBlock )
+				.Text( LOCTEXT("InputSnapLabelHorizontal", "Time"))
+				.Visibility(this, &FCurveAssetEditor::GetSnapLabelVisibilty, EOrientation::Orient_Horizontal)
+			]
+			+ SHorizontalBox::Slot()
+			[
+				SNew(SComboButton)
+				.ContentPadding(1)
+				.OnGetMenuContent(this, &FCurveAssetEditor::BuildInputSnapMenu)
+				.ToolTipText(LOCTEXT("TimeSnapToolTip", "Time Snap"))
+				.ButtonContent()
+				[
+					SNew(SEditableTextBox)
+					.Text(this, &FCurveAssetEditor::GetInputSnapText)
+					.MinDesiredWidth(40)
+					.OnTextCommitted(this, &FCurveAssetEditor::InputSnapTextComitted)
+				]
 			]
 		];
 
 	TSharedRef<SWidget> OutputSnapWidget =
 		SNew(SVerticalBox)
+		// Vertical Label
 		+ SVerticalBox::Slot()
 		.Padding(4)
 		.AutoHeight()
 		[
 			SNew(STextBlock)
-			.Text(LOCTEXT("OutputSnap", "Value Snap"))
+			.Text(LOCTEXT("OutputSnapLabelVertical", "Value Snap"))
+			.Visibility(this, &FCurveAssetEditor::GetSnapLabelVisibilty, EOrientation::Orient_Vertical)
 		]
 		+ SVerticalBox::Slot()
 		.Padding(4)
 		.AutoHeight()
 		[
-			SNew(SComboButton)
-			.ContentPadding(1)
-			.OnGetMenuContent(this, &FCurveAssetEditor::BuildOutputSnapMenu)
-			.ButtonContent()
+			SNew(SHorizontalBox)
+			// Horizontal Icon Label
+			+ SHorizontalBox::Slot()
+			.VAlign(EVerticalAlignment::VAlign_Center)
+			.AutoWidth()
+			.Padding(FMargin(0, 0, 3, 0))
 			[
-				SNew(SEditableTextBox)
-				.Text(this, &FCurveAssetEditor::GetOutputSnapText)
-				.OnTextCommitted(this, &FCurveAssetEditor::OutputSnapTextComitted)
+				SNew(STextBlock)
+				.Text(LOCTEXT("OutputSnapLabelHorizontal", "Value"))
+				.Visibility(this, &FCurveAssetEditor::GetSnapLabelVisibilty, EOrientation::Orient_Horizontal)
+			]
+			+ SHorizontalBox::Slot()
+			[
+				SNew(SComboButton)
+				.ContentPadding(1)
+				.OnGetMenuContent(this, &FCurveAssetEditor::BuildOutputSnapMenu)
+				.ToolTipText( LOCTEXT( "ValueSnapToolTip", "Value Snap" ) )
+				.ButtonContent()
+				[
+					SNew(SEditableTextBox)
+					.Text(this, &FCurveAssetEditor::GetOutputSnapText)
+					.MinDesiredWidth(40)
+					.OnTextCommitted(this, &FCurveAssetEditor::OutputSnapTextComitted)
+				]
 			]
 		];
 
@@ -360,4 +396,14 @@ TSharedPtr<FExtender> FCurveAssetEditor::GetToolbarExtender()
 	return ToolbarExtender;
 }
 
+EVisibility FCurveAssetEditor::GetSnapLabelVisibilty(EOrientation LabelOrientation) const
+{
+	if ((FMultiBoxSettings::UseSmallToolBarIcons.Get() && LabelOrientation == EOrientation::Orient_Horizontal) ||
+		(FMultiBoxSettings::UseSmallToolBarIcons.Get() == false && LabelOrientation == EOrientation::Orient_Vertical))
+	{
+		return EVisibility::Visible;
+	}
+	return EVisibility::Collapsed;
+}
+
 #undef LOCTEXT_NAMESPACE
\ No newline at end of file
diff --git a/Engine/Source/Editor/CurveAssetEditor/Private/CurveAssetEditor.h b/Engine/Source/Editor/CurveAssetEditor/Private/CurveAssetEditor.h
--- a/Engine/Source/Editor/CurveAssetEditor/Private/CurveAssetEditor.h
+++ b/Engine/Source/Editor/CurveAssetEditor/Private/CurveAssetEditor.h
@@ -69,6 +69,9 @@ private:
 	/** Builds a menu widget with options for snapping in the output domain. */
 	TSharedRef<SWidget> BuildOutputSnapMenu();
 
+	/** Get the current snap label visibility for a given orientation. */
+	EVisibility GetSnapLabelVisibilty(EOrientation LabelOrientation) const;
+
 	TSharedPtr<class SCurveEditor> TrackWidget;
 
 	/**	The tab id for the curve asset tab */
diff --git a/Engine/Source/Editor/EditorStyle/Private/SlateEditorStyle.cpp b/Engine/Source/Editor/EditorStyle/Private/SlateEditorStyle.cpp
--- a/Engine/Source/Editor/EditorStyle/Private/SlateEditorStyle.cpp
+++ b/Engine/Source/Editor/EditorStyle/Private/SlateEditorStyle.cpp
@@ -5586,17 +5586,27 @@ void FSlateEditorStyle::FStyle::SetupToolkitStyles()
 	// Rich Curve Editor
 	{
 		Set("RichCurveEditor.ZoomToFitHorizontal", new IMAGE_BRUSH("Icons/icon_CurveEditor_Horizontal_40x", Icon40x40));
+		Set("RichCurveEditor.ZoomToFitHorizontal.Small", new IMAGE_BRUSH("Icons/icon_CurveEditor_Horizontal_40x", Icon20x20));
 		Set("RichCurveEditor.ZoomToFitVertical", new IMAGE_BRUSH("Icons/icon_CurveEditor_Vertical_40x", Icon40x40));
+		Set("RichCurveEditor.ZoomToFitVertical.Small", new IMAGE_BRUSH("Icons/icon_CurveEditor_Vertical_40x", Icon20x20));
 		Set("RichCurveEditor.ZoomToFitAll", new IMAGE_BRUSH("Icons/icon_CurveEditor_All_40x", Icon40x40));
+		Set("RichCurveEditor.ZoomToFitAll.Small", new IMAGE_BRUSH("Icons/icon_CurveEditor_All_40x", Icon20x20));
 		Set("RichCurveEditor.ZoomToFitSelected", new IMAGE_BRUSH("Icons/icon_CurveEditor_Selected_40x", Icon40x40));
+		Set("RichCurveEditor.ZoomToFitSelected.Small", new IMAGE_BRUSH("Icons/icon_CurveEditor_Selected_40x", Icon20x20));
 
 		Set("RichCurveEditor.ToggleSnapping", new IMAGE_BRUSH("Icons/icon_CurveEditor_ToggleSnap_40x", Icon40x40));
+		Set("RichCurveEditor.ToggleSnapping.Small", new IMAGE_BRUSH("Icons/icon_CurveEditor_ToggleSnap_40x", Icon20x20));
 
 		Set("RichCurveEditor.InterpolationCubicAuto", new IMAGE_BRUSH("Icons/icon_CurveEditor_Auto_40x", Icon40x40));
+		Set("RichCurveEditor.InterpolationCubicAuto.Small", new IMAGE_BRUSH("Icons/icon_CurveEditor_Auto_40x", Icon20x20));
 		Set("RichCurveEditor.InterpolationCubicUser", new IMAGE_BRUSH("Icons/icon_CurveEditor_User_40x", Icon40x40));
+		Set("RichCurveEditor.InterpolationCubicUser.Small", new IMAGE_BRUSH("Icons/icon_CurveEditor_User_40x", Icon20x20));
 		Set("RichCurveEditor.InterpolationCubicBreak", new IMAGE_BRUSH("Icons/icon_CurveEditor_Break_40x", Icon40x40));
+		Set("RichCurveEditor.InterpolationCubicBreak.Small", new IMAGE_BRUSH("Icons/icon_CurveEditor_Break_40x", Icon20x20));
 		Set("RichCurveEditor.InterpolationLinear", new IMAGE_BRUSH("Icons/icon_CurveEditor_Linear_40x", Icon40x40));
+		Set("RichCurveEditor.InterpolationLinear.Small", new IMAGE_BRUSH("Icons/icon_CurveEditor_Linear_40x", Icon20x20));
 		Set("RichCurveEditor.InterpolationConstant", new IMAGE_BRUSH("Icons/icon_CurveEditor_Constant_40x", Icon40x40));
+		Set("RichCurveEditor.InterpolationConstant.Small", new IMAGE_BRUSH("Icons/icon_CurveEditor_Constant_40x", Icon20x20));
 	}
 
 	// PhAT
