Fix strange assumptions about specific curve types exiting in CurveAssetEditor as pointed out by a licensee.

diff --git a/Engine/Source/Editor/CurveAssetEditor/Private/CurveAssetEditor.cpp b/Engine/Source/Editor/CurveAssetEditor/Private/CurveAssetEditor.cpp
--- a/Engine/Source/Editor/CurveAssetEditor/Private/CurveAssetEditor.cpp
+++ b/Engine/Source/Editor/CurveAssetEditor/Private/CurveAssetEditor.cpp
@@ -131,24 +131,9 @@ TSharedRef<SDockTab> FCurveAssetEditor::SpawnTab_CurveAsset( const FSpawnTabArgs
 		];
 
 	UCurveBase* Curve = Cast<UCurveBase>(GetEditingObject());
-	UCurveFloat* FloatCurve = Cast<UCurveFloat>(Curve);
-	UCurveVector* VectorCurve = Cast<UCurveVector>(Curve);
-	UCurveLinearColor* LinearColorCurve = Cast<UCurveLinearColor>(Curve);
-
-	FCurveOwnerInterface* CurveOwner = NULL;
-	if(FloatCurve != NULL)
-	{
-		CurveOwner = FloatCurve;
-	}
-	else if(LinearColorCurve != NULL)
-	{
-		CurveOwner = LinearColorCurve;
-	}
-	else
-	{
-		CurveOwner = VectorCurve;
-	}
 
+	FCurveOwnerInterface* CurveOwner = Curve;
+	
 	if (CurveOwner != NULL)
 	{
 		check(TrackWidget.IsValid());
