Removed Tab Spawner for Color Curve Editor is your not editing the color curve.

diff --git a/Engine/Source/Editor/CurveAssetEditor/Private/CurveAssetEditor.cpp b/Engine/Source/Editor/CurveAssetEditor/Private/CurveAssetEditor.cpp
--- a/Engine/Source/Editor/CurveAssetEditor/Private/CurveAssetEditor.cpp
+++ b/Engine/Source/Editor/CurveAssetEditor/Private/CurveAssetEditor.cpp
@@ -34,10 +34,13 @@ void FCurveAssetEditor::RegisterTabSpawners(const TSharedRef<class FTabManager>&
 		.SetGroup(WorkspaceMenuCategory.ToSharedRef())
 		.SetIcon(FSlateIcon(FEditorStyle::GetStyleSetName(), "ClassIcon.CurveBase"));
 
-	InTabManager->RegisterTabSpawner(ColorCurveEditorTabId, FOnSpawnTab::CreateSP(this, &FCurveAssetEditor::SpawnTab_ColorCurveEditor))
-		.SetDisplayName(LOCTEXT("ColorCurveEditorTab", "Color Curve Editor"))
-		.SetGroup(WorkspaceMenuCategory.ToSharedRef())
-		.SetIcon(FSlateIcon(FEditorStyle::GetStyleSetName(), "ClassIcon.CurveBase"));
+	if (ColorCurveDetailsView)
+	{
+		InTabManager->RegisterTabSpawner(ColorCurveEditorTabId, FOnSpawnTab::CreateSP(this, &FCurveAssetEditor::SpawnTab_ColorCurveEditor))
+			.SetDisplayName(LOCTEXT("ColorCurveEditorTab", "Color Curve Editor"))
+			.SetGroup(WorkspaceMenuCategory.ToSharedRef())
+			.SetIcon(FSlateIcon(FEditorStyle::GetStyleSetName(), "ClassIcon.CurveBase"));
+	}
 }
 
 void FCurveAssetEditor::UnregisterTabSpawners(const TSharedRef<class FTabManager>& InTabManager)
