//UE5/Main - UE.EditorAutomation(RunTest=Editor) Win64 - AT_CurveBase_Float had dirty asset after opening
#fix prevent modifying curve on-load

diff --git a/Engine/Source/Editor/CurveAssetEditor/Private/CurveAssetEditor.cpp b/Engine/Source/Editor/CurveAssetEditor/Private/CurveAssetEditor.cpp
--- a/Engine/Source/Editor/CurveAssetEditor/Private/CurveAssetEditor.cpp
+++ b/Engine/Source/Editor/CurveAssetEditor/Private/CurveAssetEditor.cpp
@@ -117,7 +117,7 @@ struct FCurveAssetEditorTreeItem : public ICurveEditorTreeItem
 
 		TUniquePtr<FRichCurveEditorModelRaw> NewCurve = MakeUnique<FRichCurveEditorModelRaw>(static_cast<FRichCurve*>(EditInfo.CurveToEdit), CurveOwner.Get());
 		NewCurve->SetShortDisplayName(CurveName);
-		NewCurve->SetColor(CurveColor);
+		NewCurve->SetColor(CurveColor, false);
 		OutCurveModels.Add(MoveTemp(NewCurve));
 	}
 	
