Curve Editor: Implement resolving of FCurveModels for Niagara and Curve Asset Editor.

This allows command-based undo to apply its data to curves that are not currently visible in the curve editor.

Note: currently Niagara mixes Modify() calls with the command-based system. The transaction system ends up overriding this command-based data. So this change is effectively a no-op. However, this change is needed for a future follow-up change where we get rid of the transaction based system.

diff --git a/Engine/Plugins/FX/Niagara/Source/NiagaraEditorWidgets/Private/DetailCustomizations/NiagaraDataInterfaceCurveDetails.cpp b/Engine/Plugins/FX/Niagara/Source/NiagaraEditorWidgets/Private/DetailCustomizations/NiagaraDataInterfaceCurveDetails.cpp
--- a/Engine/Plugins/FX/Niagara/Source/NiagaraEditorWidgets/Private/DetailCustomizations/NiagaraDataInterfaceCurveDetails.cpp
+++ b/Engine/Plugins/FX/Niagara/Source/NiagaraEditorWidgets/Private/DetailCustomizations/NiagaraDataInterfaceCurveDetails.cpp
@@ -55,6 +55,7 @@
 #include "Misc/AxisDisplayInfo.h"
 
 #include "EditorFontGlyphs.h"
+#include "Misc/CurveEditorIntegrationUtils.h"
 #include "Widgets/Input/SSegmentedControl.h"
 
 #define LOCTEXT_NAMESPACE "NiagaraDataInterfaceCurveDetails"
@@ -148,6 +149,7 @@ struct FNiagaraCurveDetailsTreeItem : public ICurveEditorTreeItem, TSharedFromTh
 		if(CurveData.Name != NAME_None)
 		{
 			NewCurve->SetShortDisplayName(FText::FromName(CurveData.Name));
+			NewCurve->SetChannelName(CurveData.Name);
 		}
 		NewCurve->SetColor(CurveData.Color);
 		NewCurve->OnCurveModified().AddSP(this, &FNiagaraCurveDetailsTreeItem::CurveChanged);
@@ -234,6 +236,7 @@ public:
 		UNiagaraEditorSettings* EditorSettings = GetMutableDefault<UNiagaraEditorSettings>();
 		CurveEditor = MakeShared<FCurveEditor>();
 		FCurveEditorInitParams InitParams;
+		InitParams.ResolveCurveModelDelegate.BindStatic(&UE::NiagaraEditorWidgets::ResolveNiagaraCurveModel); // Needed for undo / redo.
 		CurveEditor->InitCurveEditor(InitParams);
 		TAttribute<bool>::FGetter InputSnapEnabledGetter = TAttribute<bool>::FGetter::CreateUObject(EditorSettings, &UNiagaraEditorSettings::IsCurveInputSnapEnabled);
 		TAttribute<bool>::FGetter OutputSnapEnabledGetter = TAttribute<bool>::FGetter::CreateUObject(EditorSettings, &UNiagaraEditorSettings::IsCurveOutputSnapEnabled);
diff --git a/Engine/Plugins/FX/Niagara/Source/NiagaraEditorWidgets/Private/Misc/CurveEditorIntegrationUtils.cpp b/Engine/Plugins/FX/Niagara/Source/NiagaraEditorWidgets/Private/Misc/CurveEditorIntegrationUtils.cpp
new file mode 100644
--- /dev/null
+++ b/Engine/Plugins/FX/Niagara/Source/NiagaraEditorWidgets/Private/Misc/CurveEditorIntegrationUtils.cpp
@@ -0,0 +1,38 @@
+﻿// Copyright Epic Games, Inc. All Rights Reserved.
+
+#include "CurveEditorIntegrationUtils.h"
+
+#include "NiagaraDataInterfaceCurveBase.h"
+#include "RichCurveEditorModel.h"
+#include "Modification/Resolution/CurveMetaDataIdentifiers.h"
+
+namespace UE::NiagaraEditorWidgets
+{
+TUniquePtr<FCurveModel> ResolveNiagaraCurveModel(const CurveEditor::FCurveModelLookUpArgs& InArgs)
+{
+	const CurveEditor::FCurveMetaDataIdentifiers* Metadata = InArgs.GetCurveMetaData();
+	if (!Metadata)
+	{
+		return nullptr;
+	}
+	
+	UNiagaraDataInterfaceCurveBase* CurveOwner = Cast<UNiagaraDataInterfaceCurveBase>(Metadata->Owner.Get());
+	if (!CurveOwner)
+	{
+		return nullptr;
+	}
+	
+	TArray<UNiagaraDataInterfaceCurveBase::FCurveData> CurveData;
+	CurveOwner->GetCurveData(CurveData);
+	if (CurveData.Num() == 1)
+	{
+		return MakeUnique<FRichCurveEditorModelRaw>(CurveData[0].Curve, CurveOwner);
+	}
+
+	const UNiagaraDataInterfaceCurveBase::FCurveData* Data = CurveData.FindByPredicate([Metadata](const UNiagaraDataInterfaceCurveBase::FCurveData& Data)
+	{
+		return Data.Name == Metadata->ChannelName;
+	});
+	return Data ? MakeUnique<FRichCurveEditorModelRaw>(Data->Curve, CurveOwner) : nullptr;
+}
+}
diff --git a/Engine/Plugins/FX/Niagara/Source/NiagaraEditorWidgets/Private/Misc/CurveEditorIntegrationUtils.h b/Engine/Plugins/FX/Niagara/Source/NiagaraEditorWidgets/Private/Misc/CurveEditorIntegrationUtils.h
new file mode 100644
--- /dev/null
+++ b/Engine/Plugins/FX/Niagara/Source/NiagaraEditorWidgets/Private/Misc/CurveEditorIntegrationUtils.h
@@ -0,0 +1,20 @@
+﻿// Copyright Epic Games, Inc. All Rights Reserved.
+
+#pragma once
+
+#include "Modification/Resolution/CurveModelLookUpInfo.h"
+#include "Templates/UniquePtr.h"
+
+class FCurveModel;
+
+namespace UE::NiagaraEditorWidgets
+{
+/**
+ * Tries to create a FCurveModel that was previously in the FCurveEditor but has since been removed.
+ * This is needed for Curve Editor's command based undo system.
+ * 
+ * The user may make a change to curve X, then click on curve Y, and then undo; FCurveEditor::CurveData only contains the visible curves, so it will
+ * ask us to temporarily recreate the curve for X, which we try to match based off of the metadata it had when it existed.
+ */
+TUniquePtr<FCurveModel> ResolveNiagaraCurveModel(const CurveEditor::FCurveModelLookUpArgs& InArgs);
+}
diff --git a/Engine/Plugins/FX/Niagara/Source/NiagaraEditorWidgets/Private/SNiagaraCurveOverview.cpp b/Engine/Plugins/FX/Niagara/Source/NiagaraEditorWidgets/Private/SNiagaraCurveOverview.cpp
--- a/Engine/Plugins/FX/Niagara/Source/NiagaraEditorWidgets/Private/SNiagaraCurveOverview.cpp
+++ b/Engine/Plugins/FX/Niagara/Source/NiagaraEditorWidgets/Private/SNiagaraCurveOverview.cpp
@@ -9,6 +9,7 @@
 #include "NiagaraEditorWidgetsUtilities.h"
 #include "RichCurveEditorModel.h"
 #include "SCurveEditorPanel.h"
+#include "Misc/CurveEditorIntegrationUtils.h"
 #include "Tree/CurveEditorTreeFilter.h"
 #include "Tree/ICurveEditorTreeItem.h"
 #include "Tree/SCurveEditorTree.h"
@@ -307,6 +308,7 @@ struct FNiagaraCurveOverviewTreeItem : public ICurveEditorTreeItem
 		{
 			TUniquePtr<FRichCurveEditorModelRaw> NewCurve = MakeUnique<FRichCurveEditorModelRaw>(CurveSelectionTreeNode->GetCurve(), CurveSelectionTreeNode->GetCurveDataInterface().Get());
 			NewCurve->SetShortDisplayName(CurveSelectionTreeNode->GetDisplayName());
+			NewCurve->SetChannelName(CurveSelectionTreeNode->GetCurveName());
 			NewCurve->SetColor(CurveSelectionTreeNode->GetCurveColor());
 			NewCurve->OnCurveModified().AddSP(CurveSelectionTreeNode, &FNiagaraCurveSelectionTreeNode::NotifyCurveChanged);
 			NewCurve->SetIsReadOnly(TAttribute<bool>::Create(TAttribute<bool>::FGetter::CreateSP(CurveSelectionTreeNode, &FNiagaraCurveSelectionTreeNode::GetCurveIsReadOnly)));
@@ -472,6 +474,7 @@ void SNiagaraCurveOverview::Construct(const FArguments& InArgs, TSharedRef<FNiag
 
 	CurveEditor = MakeShared<FCurveEditor>();
 	FCurveEditorInitParams InitParams;
+	InitParams.ResolveCurveModelDelegate.BindStatic(&UE::NiagaraEditorWidgets::ResolveNiagaraCurveModel); // Needed for undo / redo.
 	CurveEditor->InitCurveEditor(InitParams);
 	CurveEditor->GridLineLabelFormatXAttribute = LOCTEXT("GridXLabelFormat", "{0}");
 	CurveEditor->GetTree()->SetSortPredicate(FCurveEditorTree::FTreeItemSortPredicate::CreateStatic(&FNiagaraCurveOverviewTreeItem::CompareTreeItems));
diff --git a/Engine/Source/Editor/CurveAssetEditor/Private/CurveAssetEditor.cpp b/Engine/Source/Editor/CurveAssetEditor/Private/CurveAssetEditor.cpp
--- a/Engine/Source/Editor/CurveAssetEditor/Private/CurveAssetEditor.cpp
+++ b/Engine/Source/Editor/CurveAssetEditor/Private/CurveAssetEditor.cpp
@@ -36,6 +36,8 @@
 #include "SColorGradientEditor.h"
 #include "SCurveEditorPanel.h"
 #include "SlotBase.h"
+#include "Modification/Resolution/CurveMetaDataIdentifiers.h"
+#include "Modification/Resolution/CurveModelLookUpInfo.h"
 #include "Styling/AppStyle.h"
 #include "Styling/SlateColor.h"
 #include "Templates/Casts.h"
@@ -117,6 +119,7 @@ struct FCurveAssetEditorTreeItem : public ICurveEditorTreeItem
 
 		TUniquePtr<FRichCurveEditorModelRaw> NewCurve = MakeUnique<FRichCurveEditorModelRaw>(static_cast<FRichCurve*>(EditInfo.CurveToEdit), CurveOwner.Get());
 		NewCurve->SetShortDisplayName(CurveName);
+		NewCurve->SetChannelName(EditInfo.CurveName);
 		NewCurve->SetColor(CurveColor);
 		OutCurveModels.Add(MoveTemp(NewCurve));
 	}
@@ -128,6 +131,43 @@ private:
 	FLinearColor CurveColor;
 };
 
+namespace UE::CurveAssetEditor
+{
+/**
+ * Tries to create a FCurveModel that was previously in the FCurveEditor but has since been removed.
+ * This is needed for Curve Editor's command based undo system.
+ * 
+ * The user may make a change to curve X, then click on curve Y, and then undo; FCurveEditor::CurveData only contains the visible curves, so it will
+ * ask us to temporarily recreate the curve for X, which we try to match based off of the metadata it had when it existed.
+ */
+static TUniquePtr<FCurveModel> ResolveCurveModel(const CurveEditor::FCurveModelLookUpArgs& InArgs)
+{
+	const CurveEditor::FCurveMetaDataIdentifiers* Metadata = InArgs.GetCurveMetaData();
+	if (!Metadata)
+	{
+		return nullptr;
+	}
+	
+	UCurveBase* CurveOwner = Cast<UCurveBase>(Metadata->Owner.Get());
+	if (!CurveOwner)
+	{
+		return nullptr;
+	}
+
+	TArray<FRichCurveEditInfo> CurveInfos = CurveOwner->GetCurves();
+	if (CurveInfos.Num() == 1)
+	{
+		return MakeUnique<FRichCurveEditorModelRaw>(static_cast<FRichCurve*>(CurveInfos[0].CurveToEdit), CurveOwner);
+	}
+	const FRichCurveEditInfo* Data = CurveInfos.FindByPredicate([Metadata](const FRichCurveEditInfo& Data)
+	{
+		return Data.CurveName == Metadata->ChannelName;
+	});
+	return Data ? MakeUnique<FRichCurveEditorModelRaw>(static_cast<FRichCurve*>(Data->CurveToEdit), CurveOwner) : nullptr;
+}
+}
+
+
 void FCurveAssetEditor::RegisterTabSpawners(const TSharedRef<class FTabManager>& InTabManager)
 {
 	WorkspaceMenuCategory = InTabManager->AddLocalWorkspaceMenuCategory(LOCTEXT("WorkspaceMenu_CurveAssetEditor", "Curve Asset Editor"));
@@ -265,6 +305,7 @@ TSharedRef<SDockTab> FCurveAssetEditor::SpawnTab_CurveAsset( const FSpawnTabArgs
 
 	CurveEditor = MakeShared<FCurveEditor>();
 	FCurveEditorInitParams InitParams;
+	InitParams.ResolveCurveModelDelegate.BindStatic(&UE::CurveAssetEditor::ResolveCurveModel); // Needed for undo / redo
 	CurveEditor->InitCurveEditor(InitParams);
 	CurveEditor->GridLineLabelFormatXAttribute = LOCTEXT("GridXLabelFormat", "{0}");
 
diff --git a/Engine/Source/Editor/CurveEditor/Public/RichCurveEditorModel.h b/Engine/Source/Editor/CurveEditor/Public/RichCurveEditorModel.h
--- a/Engine/Source/Editor/CurveEditor/Public/RichCurveEditorModel.h
+++ b/Engine/Source/Editor/CurveEditor/Public/RichCurveEditorModel.h
@@ -66,6 +66,8 @@ public:
 
 	UE_API virtual TUniquePtr<IBufferedCurveModel> CreateBufferedCurveCopy() const override;
 
+	virtual UObject* GetOwningObject() const override { return WeakOwner.Get(); }
+
 	// Set a range to clamp key input values
 	void SetClampInputRange(TAttribute<TRange<double>> InClampInputRange) { ClampInputRange = InClampInputRange; }
 
